Fixing the remaining problems with missmatched module names in the build.

diff --git a/Engine/Plugins/Tests/FunctionalTestingEditor/Source/Private/EditorAutomationModule.cpp b/Engine/Plugins/Tests/FunctionalTestingEditor/Source/Private/EditorAutomationModule.cpp
--- a/Engine/Plugins/Tests/FunctionalTestingEditor/Source/Private/EditorAutomationModule.cpp
+++ b/Engine/Plugins/Tests/FunctionalTestingEditor/Source/Private/EditorAutomationModule.cpp
@@ -188,4 +188,4 @@ private:
 
 #undef LOCTEXT_NAMESPACE
 
-IMPLEMENT_MODULE(FEditorAutomationModule, EditorAutomation);
+IMPLEMENT_MODULE(FEditorAutomationModule, FunctionalTestingEditor);
