Cooking Fix

diff --git a/Engine/Plugins/Tests/FunctionalTestingEditor/Source/Private/FunctionalTestingEditorModule.cpp b/Engine/Plugins/Tests/FunctionalTestingEditor/Source/Private/FunctionalTestingEditorModule.cpp
--- a/Engine/Plugins/Tests/FunctionalTestingEditor/Source/Private/FunctionalTestingEditorModule.cpp
+++ b/Engine/Plugins/Tests/FunctionalTestingEditor/Source/Private/FunctionalTestingEditorModule.cpp
@@ -106,16 +106,19 @@ class FFunctionalTestingEditorModule : public IFunctionalTestingEditorModule
 {
 	void StartupModule()
 	{
-		// Add Automation Area to the Tools Menu
-		UToolMenu* Menu = UToolMenus::Get()->ExtendMenu("LevelEditor.MainMenu.Tools");
-		FToolMenuSection& Section = Menu->AddSection("Automation", LOCTEXT("AutomationHeading", "Automation"));
-		Section.AddMenuEntry( 
-			"TestAutomation",
-			LOCTEXT("AutomationLabel", "Test Automation"),
-			LOCTEXT("Tooltip", "Launch the Testing Automation Frontend."),
-			FSlateIcon(FAppStyle::GetAppStyleSetName(), "AutomationTools.TestAutomation"),
-			FUIAction(FExecuteAction::CreateStatic(&FFunctionalTestingEditorModule::OnShowAutomationFrontend))
-		);
+		if (FSlateApplication::IsInitialized())
+		{
+			// Add Automation Area to the Tools Menu
+			UToolMenu* Menu = UToolMenus::Get()->ExtendMenu("LevelEditor.MainMenu.Tools");
+			FToolMenuSection& Section = Menu->AddSection("Automation", LOCTEXT("AutomationHeading", "Automation"));
+			Section.AddMenuEntry( 
+				"TestAutomation",
+				LOCTEXT("AutomationLabel", "Test Automation"),
+				LOCTEXT("Tooltip", "Launch the Testing Automation Frontend."),
+				FSlateIcon(FAppStyle::GetAppStyleSetName(), "AutomationTools.TestAutomation"),
+				FUIAction(FExecuteAction::CreateStatic(&FFunctionalTestingEditorModule::OnShowAutomationFrontend))
+			);
+		}
 
 		FModuleManager::Get().OnModulesChanged().AddRaw(this, &FFunctionalTestingEditorModule::OnModulesChanged);
 
