EditorFunctionalTest actor instances are now available in PIE.

diff --git a/Engine/Plugins/Tests/FunctionalTestingEditor/Source/Public/EditorFunctionalTest.h b/Engine/Plugins/Tests/FunctionalTestingEditor/Source/Public/EditorFunctionalTest.h
--- a/Engine/Plugins/Tests/FunctionalTestingEditor/Source/Public/EditorFunctionalTest.h
+++ b/Engine/Plugins/Tests/FunctionalTestingEditor/Source/Public/EditorFunctionalTest.h
@@ -15,6 +15,11 @@ class AEditorFunctionalTest : public AFunctionalTest
 	{
 		return true;
 	}
+
+	virtual bool IsNonPIEEditorOnly() const override
+	{
+		return true;
+	}
 };
 
 /** Editor-only version of screenshot functional test */
@@ -27,4 +32,9 @@ class AEditorScreenshotFunctionalTest : public AScreenshotFunctionalTest
 	{
 		return true;
 	}
+
+	virtual bool IsNonPIEEditorOnly() const override
+	{
+		return true;
+	}
 };
